Live Link Hub: Don't enable LiveLinkExampleDevice by default.

#rnx

diff --git a/Engine/Plugins/Animation/LiveLinkExampleDevice/LiveLinkExampleDevice.uplugin b/Engine/Plugins/Animation/LiveLinkExampleDevice/LiveLinkExampleDevice.uplugin
--- a/Engine/Plugins/Animation/LiveLinkExampleDevice/LiveLinkExampleDevice.uplugin
+++ b/Engine/Plugins/Animation/LiveLinkExampleDevice/LiveLinkExampleDevice.uplugin
@@ -19,10 +19,6 @@
 		{
 			"Name": "LiveLinkDevice",
 			"Enabled": true
-		},
-		{
-			"Name": "OSC",
-			"Enabled": true
 		}
 	],
 	"Modules": [
diff --git a/Engine/Source/Programs/LiveLinkHub/Source/LiveLinkHub.Target.cs b/Engine/Source/Programs/LiveLinkHub/Source/LiveLinkHub.Target.cs
--- a/Engine/Source/Programs/LiveLinkHub/Source/LiveLinkHub.Target.cs
+++ b/Engine/Source/Programs/LiveLinkHub/Source/LiveLinkHub.Target.cs
@@ -53,7 +53,6 @@ public class LiveLinkHubTarget : TargetRules
 			"XInputDevice",
 
 			"LiveLinkDevice",
-			"LiveLinkExampleDevice",
 
 			"LiveLinkOpenVR",
 
